Add ReplicationSystemLowLevelTests program that runs the tests in ReplicationSystemTestPlugin within the low-level tests framework.

diff --git a/Engine/Build/LowLevelTests_GenProps.xml b/Engine/Build/LowLevelTests_GenProps.xml
--- a/Engine/Build/LowLevelTests_GenProps.xml
+++ b/Engine/Build/LowLevelTests_GenProps.xml
@@ -1,18 +1,25 @@
 ﻿<?xml version="1.0" encoding="utf-8"?>
 <BuildGraph xmlns:xsi="http://www.w3.org/2001/XMLSchema-instance" xsi:schemaLocation="http://www.epicgames.com/BuildGraph ../../Build/Graph/Schema.xsd" xmlns="http://www.epicgames.com/BuildGraph">
-  <Property Name="TestNames" Value="HeadlessChaos;Foundation" />
+  <Property Name="TestNames" Value="HeadlessChaos;Foundation;ReplicationSystem" />
+  <Property Name="ReplicationSystemDisabled" Value="False" />
   <Property Name="FoundationDisabled" Value="False" />
   <Property Name="HeadlessChaosDisabled" Value="False" />
+  <Property Name="ReplicationSystemShort" Value="Replication System" />
   <Property Name="FoundationShort" Value="Foundation" />
   <Property Name="HeadlessChaosShort" Value="Headless Chaos" />
+  <Property Name="ReplicationSystemTarget" Value="ReplicationSystemLowLevelTests" />
   <Property Name="FoundationTarget" Value="FoundationTests" />
   <Property Name="HeadlessChaosTarget" Value="HeadlessChaos" />
+  <Property Name="ReplicationSystemBinariesRelative" Value="Engine\Binaries" />
   <Property Name="FoundationBinariesRelative" Value="Engine\Binaries" />
   <Property Name="HeadlessChaosBinariesRelative" Value="Engine\Binaries" />
+  <Property Name="ReplicationSystemReportType" Value="console" />
   <Property Name="FoundationReportType" Value="console" />
   <Property Name="HeadlessChaosReportType" Value="console" />
+  <Option Name="RunReplicationSystemTests" DefaultValue="False" Description="Run Replication System Tests" />
   <Option Name="RunFoundationTests" DefaultValue="False" Description="Run Foundation Tests" />
   <Option Name="RunHeadlessChaosTests" DefaultValue="False" Description="Run Headless Chaos Tests" />
+  <Property Name="ReplicationSystemSupportedPlatforms" Value="Win64;Mac;Linux;Android" />
   <Property Name="FoundationSupportedPlatforms" Value="Win64;Mac;Linux;Android" />
   <Property Name="HeadlessChaosSupportedPlatforms" Value="Win64" />
 </BuildGraph>
\ No newline at end of file
diff --git a/Engine/Source/Programs/ReplicationSystemLowLevelTests/Private/ReplicationSystemLowLevelTests.cpp b/Engine/Source/Programs/ReplicationSystemLowLevelTests/Private/ReplicationSystemLowLevelTests.cpp
new file mode 100644
--- /dev/null
+++ b/Engine/Source/Programs/ReplicationSystemLowLevelTests/Private/ReplicationSystemLowLevelTests.cpp
@@ -0,0 +1,90 @@
+// Copyright Epic Games, Inc. All Rights Reserved.
+
+#include "GenericPlatform/GenericPlatformFile.h"
+#include "HAL/PlatformFileManager.h"
+#include "Iris/IrisConfig.h"
+#include "Materials/Material.h"
+#include "Misc/CommandLine.h"
+#include "Misc/ConfigCacheIni.h"
+#include "Misc/CoreDelegates.h"
+#include "Modules/ModuleManager.h"
+#include "Net/Core/Trace/Private/NetTraceInternal.h"
+#include "ProfilingDebugging/TraceAuxiliary.h"
+#include "Styling/UMGCoreStyle.h"
+#include "TestCommon/CoreUtilities.h"
+#include "TestCommon/EngineUtilities.h"
+#include "TestCommon/Initialization.h"
+#include "UObject/Package.h"
+#include "UObject/PackageResourceManager.h"
+#include "UObject/UObjectBase.h"
+
+#include <catch2/catch_test_macros.hpp>
+
+GROUP_BEFORE_GLOBAL(Catch::DefaultGroup)
+{
+	// Initialize trace
+	FString Parameter;
+	if (FParse::Value(FCommandLine::Get(), TEXT("-trace="), Parameter, false))
+	{
+		FTraceAuxiliary::Initialize(FCommandLine::Get());
+		FTraceAuxiliary::TryAutoConnect();
+	}
+
+#if UE_NET_TRACE_ENABLED
+	uint32 NetTraceVerbosity;
+	if(FParse::Value(FCommandLine::Get(), TEXT("-nettrace="), NetTraceVerbosity))
+	{
+		FNetTrace::SetTraceVerbosity(NetTraceVerbosity);
+	}
+#endif
+
+	UE::Net::SetUseIrisReplication(true);
+
+	// Use our own mock platform file implementation, we don't need to do any file I/O for these tests
+	if (IPlatformFile* WrapperFile = FPlatformFileManager::Get().GetPlatformFile(TEXT("ReplicationSystemLowLevelTestsFile")))
+	{
+		IPlatformFile* CurrentPlatformFile = &FPlatformFileManager::Get().GetPlatformFile();
+		WrapperFile->Initialize(CurrentPlatformFile, TEXT(""));
+		FPlatformFileManager::Get().SetPlatformFile(*WrapperFile);
+	}
+
+	InitAllThreadPools(true);
+	InitAsyncQueues();
+	InitTaskGraph();
+	InitStats();
+	InitGWarn();
+
+	IPackageResourceManager::Initialize();
+
+	FDelayedAutoRegisterHelper::RunAndClearDelayedAutoRegisterDelegates(EDelayedRegisterRunPhase::FileSystemReady);
+
+	if (!GetTransientPackage())
+	{
+		FModuleManager::Get().AddExtraBinarySearchPaths();
+		FConfigCacheIni::InitializeConfigSystem();
+
+		// Mock config values to allow the engine to initialize
+		GConfig->SetString(TEXT("/Script/Engine.Engine"), TEXT("AIControllerClassName"), TEXT("/Script/AIModule.AIController"), GEngineIni);
+		GConfig->SetString(TEXT("/Script/Engine.Engine"), TEXT("DefaultMaterialName"), TEXT("/Engine/Transient.MockDefaultMaterial"), GEngineIni);
+		GConfig->SetString(TEXT("/Script/Engine.Engine"), TEXT("DefaultLightFunctionMaterialName"), TEXT("/Engine/Transient.MockDefaultMaterial"), GEngineIni);
+		GConfig->SetString(TEXT("/Script/Engine.Engine"), TEXT("DefaultDeferredDecalMaterialName"), TEXT("/Engine/Transient.MockDefaultMaterial"), GEngineIni);
+		GConfig->SetString(TEXT("/Script/Engine.Engine"), TEXT("DefaultPostProcessMaterialName"), TEXT("/Engine/Transient.MockDefaultMaterial"), GEngineIni);
+
+		FModuleManager::Get().LoadModule(TEXT("CoreUObject"));
+		FCoreDelegates::OnInit.Broadcast();
+
+		FUMGCoreStyle::ResetToDefault();
+
+		// Create a mock default material to keep the material system happy
+		UMaterial* MockMaterial = NewObject<UMaterial>(GetTransientPackage(), UMaterial::StaticClass(), TEXT("MockDefaultMaterial"), RF_Transient | RF_MarkAsRootSet);
+
+		ProcessNewlyLoadedUObjects();
+	}
+
+	FModuleManager::Get().LoadModule(TEXT("IrisCore"));
+}
+
+GROUP_AFTER_GLOBAL(Catch::DefaultGroup)
+{
+	CleanupAll();
+}
diff --git a/Engine/Source/Programs/ReplicationSystemLowLevelTests/Private/ReplicationSystemLowLevelTestsFileModule.cpp b/Engine/Source/Programs/ReplicationSystemLowLevelTests/Private/ReplicationSystemLowLevelTestsFileModule.cpp
new file mode 100644
--- /dev/null
+++ b/Engine/Source/Programs/ReplicationSystemLowLevelTests/Private/ReplicationSystemLowLevelTestsFileModule.cpp
@@ -0,0 +1,59 @@
+// Copyright Epic Games, Inc. All Rights Reserved.
+
+#include "HAL/PlatformFile.h"
+#include "HAL/IPlatformFileModule.h"
+#include "Modules/ModuleInterface.h"
+#include "Modules/ModuleManager.h"
+
+class FReplicationSystemLowLevelTestsFile : public IPlatformFile
+{
+public:
+	FReplicationSystemLowLevelTestsFile() : LowerLevel(nullptr) {}
+
+	virtual bool ShouldBeUsed(IPlatformFile* Inner, const TCHAR* CmdLine) const override { return true; }
+	virtual bool Initialize(IPlatformFile* Inner, const TCHAR* CmdLine) override { LowerLevel = Inner; return true; }
+
+	virtual IPlatformFile* GetLowerLevel() { return LowerLevel; }
+	virtual void SetLowerLevel(IPlatformFile* NewLowerLevel) { LowerLevel = NewLowerLevel; }
+	virtual const TCHAR* GetName() const { return TEXT("ReplicationSystemTest"); }
+
+	virtual bool FileExists(const TCHAR* Filename) override { return false; }
+	virtual int64 FileSize(const TCHAR* Filename) override {return -1LL; }
+
+	virtual bool DeleteFile(const TCHAR* Filename) override { return false; }
+	virtual bool IsReadOnly(const TCHAR* Filename) override { return false; }
+
+	virtual bool MoveFile(const TCHAR* To, const TCHAR* From) override { return true; }
+	virtual bool SetReadOnly(const TCHAR* Filename, bool bNewReadOnlyValue) override { return false; }
+	virtual FDateTime GetTimeStamp(const TCHAR* Filename) override { return FDateTime::MinValue(); }
+	virtual void SetTimeStamp(const TCHAR* Filename, FDateTime DateTime) override {}
+	virtual FDateTime GetAccessTimeStamp(const TCHAR* Filename) override { return FDateTime::MinValue(); }
+	virtual FString GetFilenameOnDisk(const TCHAR* Filename) override { return TEXT("I.ris"); }
+	virtual IFileHandle* OpenRead(const TCHAR* Filename, bool bAllowWrite = false) override { return nullptr; }
+	virtual IFileHandle* OpenWrite(const TCHAR* Filename, bool bAppend = false, bool bAllowRead = false) override { return nullptr; }
+
+	virtual bool DirectoryExists(const TCHAR* Directory) override { return false; }
+
+	virtual bool CreateDirectory(const TCHAR* Directory) override { return false; }
+	virtual bool DeleteDirectory(const TCHAR* Directory) override { return false; }
+	virtual FFileStatData GetStatData(const TCHAR* FilenameOrDirectory) override { FFileStatData Data; Data.bIsValid = false; return Data; }
+
+	virtual bool IterateDirectory(const TCHAR* Directory, FDirectoryVisitor& Visitor) override { return true; }
+	virtual bool IterateDirectoryStat(const TCHAR* Directory, FDirectoryStatVisitor& Visitor) override { return true; }
+
+private:
+	IPlatformFile* LowerLevel;
+};
+
+
+class FReplicationSystemLowLevelTestsFileModule : public IPlatformFileModule
+{
+public:
+	virtual IPlatformFile* GetPlatformFile() override
+	{
+		static TUniquePtr<IPlatformFile> AutoDestroySingleton = MakeUnique<FReplicationSystemLowLevelTestsFile>();
+		return AutoDestroySingleton.Get();
+	}
+};
+
+IMPLEMENT_MODULE(FReplicationSystemLowLevelTestsFileModule, ReplicationSystemLowLevelTestsFile);
diff --git a/Engine/Source/Programs/ReplicationSystemLowLevelTests/ReplicationSystemLowLevelTests.Build.cs b/Engine/Source/Programs/ReplicationSystemLowLevelTests/ReplicationSystemLowLevelTests.Build.cs
new file mode 100644
--- /dev/null
+++ b/Engine/Source/Programs/ReplicationSystemLowLevelTests/ReplicationSystemLowLevelTests.Build.cs
@@ -0,0 +1,30 @@
+// Copyright Epic Games, Inc. All Rights Reserved.
+
+using UnrealBuildTool;
+
+public class ReplicationSystemLowLevelTests : TestModuleRules
+{
+	public ReplicationSystemLowLevelTests(ReadOnlyTargetRules Target) : base(Target)
+	{
+		PrivateDependencyModuleNames.AddRange(
+			new string[]
+			{
+				"Core",
+				"CoreUObject",
+				"IrisCore",
+				"ReplicationSystemTestPlugin",
+			}
+		);
+
+		PrivateIncludePathModuleNames.AddRange(
+			new string[]
+			{
+				"Engine",
+				"RHI",
+				"SlateCore"
+			}
+		);
+
+		UpdateBuildGraphPropertiesFile(new Metadata("ReplicationSystem", "Replication System"));
+	}
+}
diff --git a/Engine/Source/Programs/ReplicationSystemLowLevelTests/ReplicationSystemLowLevelTests.Target.cs b/Engine/Source/Programs/ReplicationSystemLowLevelTests/ReplicationSystemLowLevelTests.Target.cs
new file mode 100644
--- /dev/null
+++ b/Engine/Source/Programs/ReplicationSystemLowLevelTests/ReplicationSystemLowLevelTests.Target.cs
@@ -0,0 +1,21 @@
+// Copyright Epic Games, Inc. All Rights Reserved.
+
+using UnrealBuildTool;
+
+[SupportedPlatforms(UnrealPlatformClass.All)]
+public class ReplicationSystemLowLevelTestsTarget : TestTargetRules
+{
+	public ReplicationSystemLowLevelTestsTarget(TargetInfo Target) : base(Target)
+	{
+		bCompileAgainstEngine = true;
+		bUsesSlate = false;
+
+		IncludeOrderVersion = EngineIncludeOrderVersion.Latest;
+
+		// Network config
+		bWithPushModel = true;
+		bUseIris = true;
+
+		GlobalDefinitions.Add("UE_TRACE_ENABLED=1");
+	}
+}
