LLT options: Stub platform file and mock some default engine resources for low level tests that run with the engine.
- bUsePlatformFileStub = true replaces the default platform file. The default PF is saved and can be reused at any time.
The side effect is lots of warnings for missing files and object loading errors which is acceptable for low level tests.
- bMockEngineDefaults = true sets some mandatory engine defaults (materials, fonts etc) that can only be loaded from cooked assets.

diff --git a/Engine/Source/Developer/LowLevelTestsRunner/Private/TestCommon/CoreUObjectUtilities.cpp b/Engine/Source/Developer/LowLevelTestsRunner/Private/TestCommon/CoreUObjectUtilities.cpp
--- a/Engine/Source/Developer/LowLevelTestsRunner/Private/TestCommon/CoreUObjectUtilities.cpp
+++ b/Engine/Source/Developer/LowLevelTestsRunner/Private/TestCommon/CoreUObjectUtilities.cpp
@@ -12,25 +12,41 @@
 #include "UObject/PackageResourceManager.h"
 
 #if WITH_ENGINE
-#include "Framework/Application/SlateApplication.h"
-#endif
+#if UE_LLT_WITH_MOCK_ENGINE_DEFAULTS
+#include "Materials/Material.h"
+#endif // UE_LLT_WITH_MOCK_ENGINE_DEFAULTS
+#include "Styling/UMGCoreStyle.h"
+#endif //WITH_ENGINE
 
 void InitCoreUObject()
 {
 	IPackageResourceManager::Initialize();
 
+	FDelayedAutoRegisterHelper::RunAndClearDelayedAutoRegisterDelegates(EDelayedRegisterRunPhase::FileSystemReady);
+
 	if (!GetTransientPackage())
 	{
 		FModuleManager::Get().AddExtraBinarySearchPaths();
 		FConfigCacheIni::InitializeConfigSystem();
 		FPlatformFileManager::Get().InitializeNewAsyncIO();
 
+#if WITH_ENGINE && UE_LLT_WITH_MOCK_ENGINE_DEFAULTS
+		GConfig->SetString(TEXT("/Script/Engine.Engine"), TEXT("AIControllerClassName"), TEXT("/Script/AIModule.AIController"), GEngineIni);
+		GConfig->SetString(TEXT("/Script/Engine.Engine"), TEXT("DefaultMaterialName"), TEXT("/Engine/Transient.MockDefaultMaterial"), GEngineIni);
+		GConfig->SetString(TEXT("/Script/Engine.Engine"), TEXT("DefaultLightFunctionMaterialName"), TEXT("/Engine/Transient.MockDefaultMaterial"), GEngineIni);
+		GConfig->SetString(TEXT("/Script/Engine.Engine"), TEXT("DefaultDeferredDecalMaterialName"), TEXT("/Engine/Transient.MockDefaultMaterial"), GEngineIni);
+		GConfig->SetString(TEXT("/Script/Engine.Engine"), TEXT("DefaultPostProcessMaterialName"), TEXT("/Engine/Transient.MockDefaultMaterial"), GEngineIni);
+#endif // WITH_ENGINE && UE_LLT_WITH_MOCK_ENGINE_DEFAULTS
+
 		FModuleManager::Get().LoadModule(TEXT("CoreUObject"));
 		FCoreDelegates::OnInit.Broadcast();
 
 #if WITH_ENGINE
-		FSlateApplication::InitializeCoreStyle();
-#endif
+		FUMGCoreStyle::ResetToDefault();
+	#if UE_LLT_WITH_MOCK_ENGINE_DEFAULTS
+		UMaterial* MockMaterial = NewObject<UMaterial>(GetTransientPackage(), UMaterial::StaticClass(), TEXT("MockDefaultMaterial"), RF_Transient | RF_MarkAsRootSet);
+	#endif // UE_LLT_WITH_MOCK_ENGINE_DEFAULTS
+#endif // WITH_ENGINE
 
 		ProcessNewlyLoadedUObjects();
 	}
diff --git a/Engine/Source/Developer/LowLevelTestsRunner/Private/TestCommon/Initialization.cpp b/Engine/Source/Developer/LowLevelTestsRunner/Private/TestCommon/Initialization.cpp
--- a/Engine/Source/Developer/LowLevelTestsRunner/Private/TestCommon/Initialization.cpp
+++ b/Engine/Source/Developer/LowLevelTestsRunner/Private/TestCommon/Initialization.cpp
@@ -4,8 +4,13 @@
 #include "TestCommon/CoreUtilities.h"
 #include "TestCommon/CoreUObjectUtilities.h"
 #include "TestCommon/EngineUtilities.h"
+
+#include "GenericPlatform/GenericPlatformFile.h"
+#include "HAL/PlatformFileManager.h"
 #include "Misc/DelayedAutoRegister.h"
 
+static IPlatformFile* DefaultPlatformFile;
+
 void InitAllThreadPoolsEditorEx(bool MultiThreaded)
 {
 #if WITH_EDITOR
@@ -23,8 +28,32 @@ void InitStats()
 	FDelayedAutoRegisterHelper::RunAndClearDelayedAutoRegisterDelegates(EDelayedRegisterRunPhase::StatSystemReady);
 }
 
+void UsePlatformFileStubIfRequired()
+{
+#if WITH_ENGINE && UE_LLT_USE_PLATFORM_FILE_STUB
+	if (IPlatformFile* WrapperFile = FPlatformFileManager::Get().GetPlatformFile(TEXT("PlatformFileStub")))
+	{
+		IPlatformFile* CurrentPlatformFile = &FPlatformFileManager::Get().GetPlatformFile();
+		WrapperFile->Initialize(CurrentPlatformFile, TEXT(""));
+		FPlatformFileManager::Get().SetPlatformFile(*WrapperFile);
+	}
+#endif // WITH_ENGINE && UE_LLT_USE_PLATFORM_FILE_STUB
+}
+
+void SaveDefaultPlatformFile()
+{
+	DefaultPlatformFile = &FPlatformFileManager::Get().GetPlatformFile();
+}
+
+void UseDefaultPlatformFile()
+{
+	FPlatformFileManager::Get().SetPlatformFile(*DefaultPlatformFile);
+}
+
 void InitAll(bool bAllowLogging, bool bMultithreaded)
 {
+	SaveDefaultPlatformFile();
+	UsePlatformFileStubIfRequired();
 	InitAllThreadPools(bMultithreaded);
 #if WITH_ENGINE
 	InitAsyncQueues();
diff --git a/Engine/Source/Developer/LowLevelTestsRunner/Private/TestStubs/PlatformFileModuleStub.cpp b/Engine/Source/Developer/LowLevelTestsRunner/Private/TestStubs/PlatformFileModuleStub.cpp
new file mode 100644
--- /dev/null
+++ b/Engine/Source/Developer/LowLevelTestsRunner/Private/TestStubs/PlatformFileModuleStub.cpp
@@ -0,0 +1,51 @@
+// Copyright Epic Games, Inc. All Rights Reserved.
+
+#include "HAL/PlatformFile.h"
+#include "Modules/ModuleInterface.h"
+#include "Modules/ModuleManager.h"
+
+#include "HAL/IPlatformFileModule.h"
+
+class FPlatformFileStub : public IPlatformFile
+{
+public:
+	FPlatformFileStub() : LowerLevel(nullptr) {}
+	virtual bool ShouldBeUsed(IPlatformFile* Inner, const TCHAR* CmdLine) const override { return true; }
+	virtual bool Initialize(IPlatformFile* Inner, const TCHAR* CmdLine) override { LowerLevel = Inner; return true; }
+	virtual IPlatformFile* GetLowerLevel() { return LowerLevel; }
+	virtual void SetLowerLevel(IPlatformFile* NewLowerLevel) { LowerLevel = NewLowerLevel; }
+	virtual const TCHAR* GetName() const { return TEXT("LLTPlatformFileStub"); }
+	virtual bool FileExists(const TCHAR* Filename) override { return false; }
+	virtual int64 FileSize(const TCHAR* Filename) override { return -1LL; }
+	virtual bool DeleteFile(const TCHAR* Filename) override { return false; }
+	virtual bool IsReadOnly(const TCHAR* Filename) override { return false; }
+	virtual bool MoveFile(const TCHAR* To, const TCHAR* From) override { return false; }
+	virtual bool SetReadOnly(const TCHAR* Filename, bool bNewReadOnlyValue) override { return false; }
+	virtual FDateTime GetTimeStamp(const TCHAR* Filename) override { return FDateTime::MinValue(); }
+	virtual void SetTimeStamp(const TCHAR* Filename, FDateTime DateTime) override {}
+	virtual FDateTime GetAccessTimeStamp(const TCHAR* Filename) override { return FDateTime::MinValue(); }
+	virtual FString GetFilenameOnDisk(const TCHAR* Filename) override { return TEXT(""); }
+	virtual IFileHandle* OpenRead(const TCHAR* Filename, bool bAllowWrite = false) override { return nullptr; }
+	virtual IFileHandle* OpenWrite(const TCHAR* Filename, bool bAppend = false, bool bAllowRead = false) override { return nullptr; }
+	virtual bool DirectoryExists(const TCHAR* Directory) override { return false; }
+	virtual bool CreateDirectory(const TCHAR* Directory) override { return false; }
+	virtual bool DeleteDirectory(const TCHAR* Directory) override { return false; }
+	virtual FFileStatData GetStatData(const TCHAR* FilenameOrDirectory) override { FFileStatData Data; Data.bIsValid = false; return Data; }
+	virtual bool IterateDirectory(const TCHAR* Directory, FDirectoryVisitor& Visitor) override { return true; }
+	virtual bool IterateDirectoryStat(const TCHAR* Directory, FDirectoryStatVisitor& Visitor) override { return true; }
+private:
+	IPlatformFile* LowerLevel;
+};
+
+class FPlatformFileModuleStub : public IPlatformFileModule
+{
+public:
+	virtual IPlatformFile* GetPlatformFile() override
+	{
+		static TUniquePtr<IPlatformFile> AutoDestroySingleton = MakeUnique<FPlatformFileStub>();
+		return AutoDestroySingleton.Get();
+	}
+};
+
+
+IMPLEMENT_MODULE(FPlatformFileModuleStub, PlatformFileStub);
\ No newline at end of file
diff --git a/Engine/Source/Developer/LowLevelTestsRunner/Public/TestCommon/Initialization.h b/Engine/Source/Developer/LowLevelTestsRunner/Public/TestCommon/Initialization.h
--- a/Engine/Source/Developer/LowLevelTestsRunner/Public/TestCommon/Initialization.h
+++ b/Engine/Source/Developer/LowLevelTestsRunner/Public/TestCommon/Initialization.h
@@ -4,6 +4,10 @@
 
 #include "CoreMinimal.h"
 
+void SaveDefaultPlatformFile();
+void UsePlatformFileStubIfRequired();
+void UseDefaultPlatformFile();
+
 void InitAllThreadPoolsEditorEx(bool MultiThreaded);
 
 void InitOutputDevicesEx();
diff --git a/Engine/Source/Developer/Zen/Private/ZenServerInterface.cpp b/Engine/Source/Developer/Zen/Private/ZenServerInterface.cpp
--- a/Engine/Source/Developer/Zen/Private/ZenServerInterface.cpp
+++ b/Engine/Source/Developer/Zen/Private/ZenServerInterface.cpp
@@ -134,8 +134,10 @@ ZenServerState::ZenServerState()
 	}
 #endif
 
+#if PLATFORM_WINDOWS || PLATFORM_UNIX || PLATFORM_MAC
 	m_hMapFile = hMap;
 	m_Data = reinterpret_cast<const ZenServerEntry*>(pBuf);
+#endif
 }
 
 ZenServerState::~ZenServerState()
diff --git a/Engine/Source/Programs/HeadlessChaos/HeadlessChaos.Target.cs b/Engine/Source/Programs/HeadlessChaos/HeadlessChaos.Target.cs
--- a/Engine/Source/Programs/HeadlessChaos/HeadlessChaos.Target.cs
+++ b/Engine/Source/Programs/HeadlessChaos/HeadlessChaos.Target.cs
@@ -9,17 +9,8 @@ public class HeadlessChaosTarget : TestTargetRules
 	{
 		ExeBinariesSubFolder = LaunchModuleName = "HeadlessChaos";
 
-		// HeadlessChaos doesn't ever compile with the engine linked in
-		bCompileAgainstEngine = false;
-		bCompileAgainstCoreUObject = true;
-
         bHasExports = false;
 
-        bUseLoggingInShipping = true;
-
-		// This app is a console application, not a Windows app (sets entry point to main(), instead of WinMain())
-		bIsBuildingConsoleApplication = true;
-
 		GlobalDefinitions.Add("CHAOS_SERIALIZE_OUT=1");
 
 		// our gtest does not have debug crt libs
diff --git a/Engine/Source/Programs/ReplicationSystemLowLevelTests/Private/ReplicationSystemLowLevelTests.cpp b/Engine/Source/Programs/ReplicationSystemLowLevelTests/Private/ReplicationSystemLowLevelTests.cpp
--- a/Engine/Source/Programs/ReplicationSystemLowLevelTests/Private/ReplicationSystemLowLevelTests.cpp
+++ b/Engine/Source/Programs/ReplicationSystemLowLevelTests/Private/ReplicationSystemLowLevelTests.cpp
@@ -1,22 +1,12 @@
 // Copyright Epic Games, Inc. All Rights Reserved.
 
-#include "GenericPlatform/GenericPlatformFile.h"
-#include "HAL/PlatformFileManager.h"
 #include "Iris/IrisConfig.h"
-#include "Materials/Material.h"
 #include "Misc/CommandLine.h"
-#include "Misc/ConfigCacheIni.h"
-#include "Misc/CoreDelegates.h"
 #include "Modules/ModuleManager.h"
 #include "Net/Core/Trace/Private/NetTraceInternal.h"
 #include "ProfilingDebugging/TraceAuxiliary.h"
-#include "Styling/UMGCoreStyle.h"
-#include "TestCommon/CoreUtilities.h"
-#include "TestCommon/EngineUtilities.h"
+
 #include "TestCommon/Initialization.h"
-#include "UObject/Package.h"
-#include "UObject/PackageResourceManager.h"
-#include "UObject/UObjectBase.h"
 
 #include <catch2/catch_test_macros.hpp>
 
@@ -40,46 +30,7 @@ GROUP_BEFORE_GLOBAL(Catch::DefaultGroup)
 
 	UE::Net::SetUseIrisReplication(true);
 
-	// Use our own mock platform file implementation, we don't need to do any file I/O for these tests
-	if (IPlatformFile* WrapperFile = FPlatformFileManager::Get().GetPlatformFile(TEXT("ReplicationSystemLowLevelTestsFile")))
-	{
-		IPlatformFile* CurrentPlatformFile = &FPlatformFileManager::Get().GetPlatformFile();
-		WrapperFile->Initialize(CurrentPlatformFile, TEXT(""));
-		FPlatformFileManager::Get().SetPlatformFile(*WrapperFile);
-	}
-
-	InitAllThreadPools(true);
-	InitAsyncQueues();
-	InitTaskGraph();
-	InitStats();
-	InitGWarn();
-
-	IPackageResourceManager::Initialize();
-
-	FDelayedAutoRegisterHelper::RunAndClearDelayedAutoRegisterDelegates(EDelayedRegisterRunPhase::FileSystemReady);
-
-	if (!GetTransientPackage())
-	{
-		FModuleManager::Get().AddExtraBinarySearchPaths();
-		FConfigCacheIni::InitializeConfigSystem();
-
-		// Mock config values to allow the engine to initialize
-		GConfig->SetString(TEXT("/Script/Engine.Engine"), TEXT("AIControllerClassName"), TEXT("/Script/AIModule.AIController"), GEngineIni);
-		GConfig->SetString(TEXT("/Script/Engine.Engine"), TEXT("DefaultMaterialName"), TEXT("/Engine/Transient.MockDefaultMaterial"), GEngineIni);
-		GConfig->SetString(TEXT("/Script/Engine.Engine"), TEXT("DefaultLightFunctionMaterialName"), TEXT("/Engine/Transient.MockDefaultMaterial"), GEngineIni);
-		GConfig->SetString(TEXT("/Script/Engine.Engine"), TEXT("DefaultDeferredDecalMaterialName"), TEXT("/Engine/Transient.MockDefaultMaterial"), GEngineIni);
-		GConfig->SetString(TEXT("/Script/Engine.Engine"), TEXT("DefaultPostProcessMaterialName"), TEXT("/Engine/Transient.MockDefaultMaterial"), GEngineIni);
-
-		FModuleManager::Get().LoadModule(TEXT("CoreUObject"));
-		FCoreDelegates::OnInit.Broadcast();
-
-		FUMGCoreStyle::ResetToDefault();
-
-		// Create a mock default material to keep the material system happy
-		UMaterial* MockMaterial = NewObject<UMaterial>(GetTransientPackage(), UMaterial::StaticClass(), TEXT("MockDefaultMaterial"), RF_Transient | RF_MarkAsRootSet);
-
-		ProcessNewlyLoadedUObjects();
-	}
+	InitAll(true, true);
 
 	FModuleManager::Get().LoadModule(TEXT("IrisCore"));
 }
diff --git a/Engine/Source/Programs/ReplicationSystemLowLevelTests/Private/ReplicationSystemLowLevelTestsFileModule.cpp b/Engine/Source/Programs/ReplicationSystemLowLevelTests/Private/ReplicationSystemLowLevelTestsFileModule.cpp
deleted file mode 100644
diff --git a/Engine/Source/Programs/ReplicationSystemLowLevelTests/ReplicationSystemLowLevelTests.Target.cs b/Engine/Source/Programs/ReplicationSystemLowLevelTests/ReplicationSystemLowLevelTests.Target.cs
--- a/Engine/Source/Programs/ReplicationSystemLowLevelTests/ReplicationSystemLowLevelTests.Target.cs
+++ b/Engine/Source/Programs/ReplicationSystemLowLevelTests/ReplicationSystemLowLevelTests.Target.cs
@@ -10,6 +10,9 @@ public class ReplicationSystemLowLevelTestsTarget : TestTargetRules
 		bCompileAgainstEngine = true;
 		bUsesSlate = false;
 
+		bUsePlatformFileStub = true;
+		bMockEngineDefaults = true;
+
 		IncludeOrderVersion = EngineIncludeOrderVersion.Latest;
 
 		// Network config
diff --git a/Engine/Source/Programs/UnrealBuildTool/Configuration/TestTargetRules.cs b/Engine/Source/Programs/UnrealBuildTool/Configuration/TestTargetRules.cs
--- a/Engine/Source/Programs/UnrealBuildTool/Configuration/TestTargetRules.cs
+++ b/Engine/Source/Programs/UnrealBuildTool/Configuration/TestTargetRules.cs
@@ -86,9 +86,8 @@ namespace UnrealBuildTool
 
 		/// <summary>
 		/// If set to true, it will not compile against engine even if "Engine" is in the dependency graph.
-		/// Note: Currently set to true because it requires cooked Engine assets.
 		/// </summary>
-		public static bool bNeverCompileAgainstEngine = true;
+		public static bool bNeverCompileAgainstEngine = false;
 
 		/// <summary>
 		/// Test target override for bCompileAgainstEditor.
@@ -106,7 +105,39 @@ namespace UnrealBuildTool
 		public static bool bNeverCompileAgainstEditor = false;
 
 		/// <summary>
-		/// Constructor for TestTargetRules as own target.
+		/// Whether to stub the platform file.
+		/// </summary>
+		public bool bUsePlatformFileStub
+		{ 
+			get { return bUsePlatformFileStubPrivate; }
+			set
+			{
+				bUsePlatformFileStubPrivate = value;
+				GlobalDefinitions.Remove("UE_LLT_USE_PLATFORM_FILE_STUB=0");
+				GlobalDefinitions.Remove("UE_LLT_USE_PLATFORM_FILE_STUB=1");
+				GlobalDefinitions.Add($"UE_LLT_USE_PLATFORM_FILE_STUB={Convert.ToInt32(bUsePlatformFileStubPrivate)}");
+			}
+		}
+		private bool bUsePlatformFileStubPrivate = false;
+
+		/// <summary>
+		/// Whether to mock engine default instances for materials, AI controller etc.
+		/// </summary>
+		public bool bMockEngineDefaults
+		{
+			get { return bMockEngineDefaultsPrivate; }
+			set
+			{
+				bMockEngineDefaultsPrivate = value;
+				GlobalDefinitions.Remove("UE_LLT_WITH_MOCK_ENGINE_DEFAULTS=0");
+				GlobalDefinitions.Remove("UE_LLT_WITH_MOCK_ENGINE_DEFAULTS=1");
+				GlobalDefinitions.Add($"UE_LLT_WITH_MOCK_ENGINE_DEFAULTS={Convert.ToInt32(bMockEngineDefaultsPrivate)}");
+			}
+		}
+		private bool bMockEngineDefaultsPrivate = false;
+
+		/// <summary>
+		/// Constructor that explicit targets can inherit from.
 		/// </summary>
 		/// <param name="Target"></param>
 		public TestTargetRules(TargetInfo Target) : base(Target)
@@ -119,6 +150,10 @@ namespace UnrealBuildTool
 			{
 				bBuildInSolutionByDefault = true;
 				SolutionDirectory = "Programs/LowLevelTests";
+
+				// Default to true for explicit targets to reduce compilation times.
+				// Selective module compilation will automatically detect if Engine is required based on Engine include files in tests.
+				bNeverCompileAgainstEngine = true;
 			}
 
 			if (Target.Platform == UnrealTargetPlatform.Win64)
@@ -263,6 +298,10 @@ namespace UnrealBuildTool
 
 			GlobalDefinitions.Add("STATS=0");
 			GlobalDefinitions.Add("TEST_FOR_VALID_FILE_SYSTEM_MEMORY=0");
+			
+			// LLT Globals
+			GlobalDefinitions.Add("UE_LLT_USE_PLATFORM_FILE_STUB=0");
+			GlobalDefinitions.Add("UE_LLT_WITH_MOCK_ENGINE_DEFAULTS=0");
 
 			// Platform specific setup
 			if (Target.Platform == UnrealTargetPlatform.Android)
