Replication system tests: suppress some engine warnings during initialization that are unrelated to the replication tests.

diff --git a/Engine/Source/Programs/ReplicationSystemLowLevelTests/Private/ReplicationSystemLowLevelTests.cpp b/Engine/Source/Programs/ReplicationSystemLowLevelTests/Private/ReplicationSystemLowLevelTests.cpp
--- a/Engine/Source/Programs/ReplicationSystemLowLevelTests/Private/ReplicationSystemLowLevelTests.cpp
+++ b/Engine/Source/Programs/ReplicationSystemLowLevelTests/Private/ReplicationSystemLowLevelTests.cpp
@@ -1,11 +1,13 @@
 // Copyright Epic Games, Inc. All Rights Reserved.
 
 #include "Iris/IrisConfig.h"
+#include "Logging/LogScopedVerbosityOverride.h"
 #include "Misc/CommandLine.h"
 #include "Modules/ModuleManager.h"
 #include "Net/Core/Trace/Private/NetTraceInternal.h"
 #include "ProfilingDebugging/TraceAuxiliary.h"
-
+#include "SlateGlobals.h"
+#include "Styling/SlateWidgetStyleContainerBase.h
 #include "TestCommon/Initialization.h"
 
 #include <catch2/catch_test_macros.hpp>
@@ -30,7 +32,13 @@ GROUP_BEFORE_GLOBAL(Catch::DefaultGroup)
 
 	UE::Net::SetUseIrisReplication(true);
 
-	InitAll(true, true);
+	{
+		// Silence some warnings during initialization unrelated to the replication system
+		LOG_SCOPE_VERBOSITY_OVERRIDE(LogSlate, ELogVerbosity::Error);
+		LOG_SCOPE_VERBOSITY_OVERRIDE(LogSlateStyle, ELogVerbosity::Error);
+		LOG_SCOPE_VERBOSITY_OVERRIDE(LogUObjectGlobals, ELogVerbosity::Error);
+		InitAll(true, true);
+	}
 
 	FModuleManager::Get().LoadModule(TEXT("IrisCore"));
 }
diff --git a/Engine/Source/Programs/ReplicationSystemLowLevelTests/ReplicationSystemLowLevelTests.Build.cs b/Engine/Source/Programs/ReplicationSystemLowLevelTests/ReplicationSystemLowLevelTests.Build.cs
--- a/Engine/Source/Programs/ReplicationSystemLowLevelTests/ReplicationSystemLowLevelTests.Build.cs
+++ b/Engine/Source/Programs/ReplicationSystemLowLevelTests/ReplicationSystemLowLevelTests.Build.cs
@@ -21,7 +21,9 @@ public class ReplicationSystemLowLevelTests : TestModuleRules
 			{
 				"Engine",
 				"RHI",
-				"SlateCore"
+				"SlateCore",
+				"InputCore",
+				"ApplicationCore"
 			}
 		);
 
