Fix typo

diff --git a/Engine/Source/Programs/ReplicationSystemLowLevelTests/Private/ReplicationSystemLowLevelTests.cpp b/Engine/Source/Programs/ReplicationSystemLowLevelTests/Private/ReplicationSystemLowLevelTests.cpp
--- a/Engine/Source/Programs/ReplicationSystemLowLevelTests/Private/ReplicationSystemLowLevelTests.cpp
+++ b/Engine/Source/Programs/ReplicationSystemLowLevelTests/Private/ReplicationSystemLowLevelTests.cpp
@@ -7,7 +7,7 @@
 #include "Net/Core/Trace/Private/NetTraceInternal.h"
 #include "ProfilingDebugging/TraceAuxiliary.h"
 #include "SlateGlobals.h"
-#include "Styling/SlateWidgetStyleContainerBase.h
+#include "Styling/SlateWidgetStyleContainerBase.h"
 #include "TestCommon/Initialization.h"
 
 #include <catch2/catch_test_macros.hpp>
