Run EDelayedRegisterRunPhase::PreObjectSystemReady and EDelayedRegisterRunPhase::ObjectSystemReady in the low level test runner's InitCoreUObject helper.
Replication system tests: initialize some ini settings in the PreObjectSystemReady phase to prevent "Generating descriptor for struct that has custom serialization" warnings.

diff --git a/Engine/Source/Developer/LowLevelTestsRunner/Private/TestCommon/CoreUObjectUtilities.cpp b/Engine/Source/Developer/LowLevelTestsRunner/Private/TestCommon/CoreUObjectUtilities.cpp
--- a/Engine/Source/Developer/LowLevelTestsRunner/Private/TestCommon/CoreUObjectUtilities.cpp
+++ b/Engine/Source/Developer/LowLevelTestsRunner/Private/TestCommon/CoreUObjectUtilities.cpp
@@ -87,7 +87,11 @@ void InitCoreUObject()
 	#endif // UE_LLT_USE_PLATFORM_FILE_STUB
 #endif // WITH_ENGINE
 
+		FDelayedAutoRegisterHelper::RunAndClearDelayedAutoRegisterDelegates(EDelayedRegisterRunPhase::PreObjectSystemReady);
+
 		ProcessNewlyLoadedUObjects();
+
+		FDelayedAutoRegisterHelper::RunAndClearDelayedAutoRegisterDelegates(EDelayedRegisterRunPhase::ObjectSystemReady);
 	}
 }
 
diff --git a/Engine/Source/Programs/ReplicationSystemLowLevelTests/Private/ReplicationSystemLowLevelTests.cpp b/Engine/Source/Programs/ReplicationSystemLowLevelTests/Private/ReplicationSystemLowLevelTests.cpp
--- a/Engine/Source/Programs/ReplicationSystemLowLevelTests/Private/ReplicationSystemLowLevelTests.cpp
+++ b/Engine/Source/Programs/ReplicationSystemLowLevelTests/Private/ReplicationSystemLowLevelTests.cpp
@@ -3,6 +3,8 @@
 #include "Iris/IrisConfig.h"
 #include "Logging/LogScopedVerbosityOverride.h"
 #include "Misc/CommandLine.h"
+#include "Misc/ConfigCacheIni.h"
+#include "Misc/DelayedAutoRegister.h"
 #include "Modules/ModuleManager.h"
 #include "Net/Core/Trace/Private/NetTraceInternal.h"
 #include "ProfilingDebugging/TraceAuxiliary.h"
@@ -12,6 +14,25 @@
 
 #include <catch2/catch_test_macros.hpp>
 
+namespace
+{
+	FDelayedAutoRegisterHelper DelayedAutoRegister_PreObjectSystemReady(EDelayedRegisterRunPhase::PreObjectSystemReady, []
+	{
+		// Set up Iris config here in code so we can avoid having to load ini files from disk.
+		const TArray<FString> StructsSupportingNetSerializerList =
+		{
+			TEXT("(StructName=RootMotionSource)"),
+			TEXT("(StructName=RootMotionSource_ConstantForce)"),
+			TEXT("(StructName=RootMotionSource_RadialForce)"),
+			TEXT("(StructName=RootMotionSource_MoveToForce)"),
+			TEXT("(StructName=RootMotionSource_MoveToDynamicForce)"),
+			TEXT("(StructName=RootMotionSource_JumpForce)")
+		};
+
+		GConfig->SetArray(TEXT("/Script/IrisCore.ReplicationStateDescriptorConfig"), TEXT("SupportsStructNetSerializerList"), StructsSupportingNetSerializerList, GEngineIni);
+	});
+}
+
 GROUP_BEFORE_GLOBAL(Catch::DefaultGroup)
 {
 	// Initialize trace
