Set nullrhi and unattended flags for ReplicationSystemLowLevelTests. Skips unnecessary code that could cause issues on some platforms.

diff --git a/Engine/Source/Programs/ReplicationSystemLowLevelTests/Private/ReplicationSystemLowLevelTests.cpp b/Engine/Source/Programs/ReplicationSystemLowLevelTests/Private/ReplicationSystemLowLevelTests.cpp
--- a/Engine/Source/Programs/ReplicationSystemLowLevelTests/Private/ReplicationSystemLowLevelTests.cpp
+++ b/Engine/Source/Programs/ReplicationSystemLowLevelTests/Private/ReplicationSystemLowLevelTests.cpp
@@ -35,6 +35,9 @@ namespace
 
 GROUP_BEFORE_GLOBAL(Catch::DefaultGroup)
 {
+	// Add command-line options always needed for the tests
+	FCommandLine::Append(TEXT(" -nullrhi -unattended"));
+
 	// Initialize trace
 	FString Parameter;
 	if (FParse::Value(FCommandLine::Get(), TEXT("-trace="), Parameter, false))
