Add UE_NET_TEST_FAKE_REP_TAGS to replication system low level tests

#rnx

diff --git a/Engine/Source/Programs/ReplicationSystemLowLevelTests/ReplicationSystemLowLevelTests.Target.cs b/Engine/Source/Programs/ReplicationSystemLowLevelTests/ReplicationSystemLowLevelTests.Target.cs
--- a/Engine/Source/Programs/ReplicationSystemLowLevelTests/ReplicationSystemLowLevelTests.Target.cs
+++ b/Engine/Source/Programs/ReplicationSystemLowLevelTests/ReplicationSystemLowLevelTests.Target.cs
@@ -20,5 +20,6 @@ public class ReplicationSystemLowLevelTestsTarget : TestTargetRules
 		bUseIris = true;
 
 		GlobalDefinitions.Add("UE_TRACE_ENABLED=1");
+		GlobalDefinitions.Add("UE_NET_TEST_FAKE_REP_TAGS=1");
 	}
 }
