More test fixes.

diff --git a/Engine/Source/Programs/NetworkPredictionTests/NetworkPredictionTests.Target.cs b/Engine/Source/Programs/NetworkPredictionTests/NetworkPredictionTests.Target.cs
--- a/Engine/Source/Programs/NetworkPredictionTests/NetworkPredictionTests.Target.cs
+++ b/Engine/Source/Programs/NetworkPredictionTests/NetworkPredictionTests.Target.cs
@@ -7,6 +7,7 @@ public class NetworkPredictionTestsTarget : TestTargetRules
 	public NetworkPredictionTestsTarget(TargetInfo Target) : base(Target)
 	{
 		bCompileAgainstEngine = true;
+		bCompileAgainstApplicationCore = true;
 		bUsesSlate = false;
 
 		bUsePlatformFileStub = true;
diff --git a/Engine/Source/Programs/ReplicationSystemLowLevelTests/ReplicationSystemLowLevelTests.Target.cs b/Engine/Source/Programs/ReplicationSystemLowLevelTests/ReplicationSystemLowLevelTests.Target.cs
--- a/Engine/Source/Programs/ReplicationSystemLowLevelTests/ReplicationSystemLowLevelTests.Target.cs
+++ b/Engine/Source/Programs/ReplicationSystemLowLevelTests/ReplicationSystemLowLevelTests.Target.cs
@@ -8,6 +8,7 @@ public class ReplicationSystemLowLevelTestsTarget : TestTargetRules
 	public ReplicationSystemLowLevelTestsTarget(TargetInfo Target) : base(Target)
 	{
 		bCompileAgainstEngine = true;
+		bCompileAgainstApplicationCore = true;
 		bUsesSlate = false;
 
 		bUsePlatformFileStub = true;
