Iris - Disable load time profile tracing in our test programs.

diff --git a/Engine/Source/Programs/ReplicationSystemLowLevelTests/ReplicationSystemLowLevelTests.Target.cs b/Engine/Source/Programs/ReplicationSystemLowLevelTests/ReplicationSystemLowLevelTests.Target.cs
--- a/Engine/Source/Programs/ReplicationSystemLowLevelTests/ReplicationSystemLowLevelTests.Target.cs
+++ b/Engine/Source/Programs/ReplicationSystemLowLevelTests/ReplicationSystemLowLevelTests.Target.cs
@@ -22,5 +22,7 @@ public class ReplicationSystemLowLevelTestsTarget : TestTargetRules
 
 		bEnableTrace = true;
 		GlobalDefinitions.Add("UE_NET_TEST_FAKE_REP_TAGS=1");
+		// Load time profiling brings object construction to a crawl.
+		GlobalDefinitions.Add("LOADTIMEPROFILERTRACE_ENABLED=0");
 	}
 }
diff --git a/Engine/Source/Programs/ReplicationSystemTest/ReplicationSystemTest.Target.cs b/Engine/Source/Programs/ReplicationSystemTest/ReplicationSystemTest.Target.cs
--- a/Engine/Source/Programs/ReplicationSystemTest/ReplicationSystemTest.Target.cs
+++ b/Engine/Source/Programs/ReplicationSystemTest/ReplicationSystemTest.Target.cs
@@ -41,6 +41,8 @@ public class ReplicationSystemTestTarget : TargetRules
 
 		bEnableTrace = true;
 		GlobalDefinitions.Add("UE_NET_TEST_FAKE_REP_TAGS=1");
+		// Load time profiling brings object construction to a crawl.
+		GlobalDefinitions.Add("LOADTIMEPROFILERTRACE_ENABLED=0");
 	}
 }
 
