Load ReplicationSystemTestPlugin module in ReplicationSystemLowLevelTests to prevent asserts.

diff --git a/Engine/Source/Programs/ReplicationSystemLowLevelTests/Private/ReplicationSystemLowLevelTests.cpp b/Engine/Source/Programs/ReplicationSystemLowLevelTests/Private/ReplicationSystemLowLevelTests.cpp
--- a/Engine/Source/Programs/ReplicationSystemLowLevelTests/Private/ReplicationSystemLowLevelTests.cpp
+++ b/Engine/Source/Programs/ReplicationSystemLowLevelTests/Private/ReplicationSystemLowLevelTests.cpp
@@ -69,7 +69,10 @@ GROUP_BEFORE_GLOBAL(Catch::DefaultGroup)
 		InitAll(true, true);
 	}
 
-	FModuleManager::Get().LoadModule(TEXT("IrisCore"));
+	for (const FName ModuleToLoad : {"IrisCore", "ReplicationSystemTestPlugin"})
+	{
+		FModuleManager::Get().LoadModule(ModuleToLoad);
+	}
 }
 
 GROUP_AFTER_GLOBAL(Catch::DefaultGroup)
