Fix project generation for users without access to the RTFM compiler
#rnx

diff --git a/Engine/Source/Programs/ReplicationSystemLowLevelTests/ReplicationSystemLowLevelTests.Target.cs b/Engine/Source/Programs/ReplicationSystemLowLevelTests/ReplicationSystemLowLevelTests.Target.cs
--- a/Engine/Source/Programs/ReplicationSystemLowLevelTests/ReplicationSystemLowLevelTests.Target.cs
+++ b/Engine/Source/Programs/ReplicationSystemLowLevelTests/ReplicationSystemLowLevelTests.Target.cs
@@ -21,7 +21,7 @@ public class ReplicationSystemLowLevelTestsTarget : TestTargetRules
 
 		bEnableTrace = true;
 
-		if (Target.Platform.IsInGroup(UnrealPlatformGroup.Linux) || Target.Platform == UnrealTargetPlatform.Win64)
+		if (!bGenerateProjectFiles && (Target.Platform.IsInGroup(UnrealPlatformGroup.Linux) || Target.Platform == UnrealTargetPlatform.Win64))
 		{
 			bUseAutoRTFMCompiler = true;
 		}
diff --git a/Engine/Source/Programs/ReplicationSystemTest/ReplicationSystemTest.Target.cs b/Engine/Source/Programs/ReplicationSystemTest/ReplicationSystemTest.Target.cs
--- a/Engine/Source/Programs/ReplicationSystemTest/ReplicationSystemTest.Target.cs
+++ b/Engine/Source/Programs/ReplicationSystemTest/ReplicationSystemTest.Target.cs
@@ -40,7 +40,7 @@ public class ReplicationSystemTestTarget : TargetRules
 
 		bEnableTrace = true;
 
-		if (Target.Platform.IsInGroup(UnrealPlatformGroup.Linux) || Target.Platform == UnrealTargetPlatform.Win64)
+		if (!bGenerateProjectFiles && (Target.Platform.IsInGroup(UnrealPlatformGroup.Linux) || Target.Platform == UnrealTargetPlatform.Win64))
 		{
 			bUseAutoRTFMCompiler = true;
 		}
