Use default compiler when building ReplicationSystem test programs.

#rnx

diff --git a/Engine/Source/Programs/ReplicationSystemLowLevelTests/ReplicationSystemLowLevelTests.Target.cs b/Engine/Source/Programs/ReplicationSystemLowLevelTests/ReplicationSystemLowLevelTests.Target.cs
--- a/Engine/Source/Programs/ReplicationSystemLowLevelTests/ReplicationSystemLowLevelTests.Target.cs
+++ b/Engine/Source/Programs/ReplicationSystemLowLevelTests/ReplicationSystemLowLevelTests.Target.cs
@@ -23,7 +23,7 @@ public class ReplicationSystemLowLevelTestsTarget : TestTargetRules
 
 		if (!bGenerateProjectFiles && (Target.Platform.IsInGroup(UnrealPlatformGroup.Linux) || Target.Platform == UnrealTargetPlatform.Win64))
 		{
-			bUseAutoRTFMCompiler = true;
+			//bUseAutoRTFMCompiler = true;
 		}
 
 		GlobalDefinitions.Add("UE_NET_TEST_FAKE_REP_TAGS=1");
diff --git a/Engine/Source/Programs/ReplicationSystemTest/ReplicationSystemTest.Target.cs b/Engine/Source/Programs/ReplicationSystemTest/ReplicationSystemTest.Target.cs
--- a/Engine/Source/Programs/ReplicationSystemTest/ReplicationSystemTest.Target.cs
+++ b/Engine/Source/Programs/ReplicationSystemTest/ReplicationSystemTest.Target.cs
@@ -42,7 +42,7 @@ public class ReplicationSystemTestTarget : TargetRules
 
 		if (!bGenerateProjectFiles && (Target.Platform.IsInGroup(UnrealPlatformGroup.Linux) || Target.Platform == UnrealTargetPlatform.Win64))
 		{
-			bUseAutoRTFMCompiler = true;
+			//bUseAutoRTFMCompiler = true;
 		}
 		
 		GlobalDefinitions.Add("UE_NET_TEST_FAKE_REP_TAGS=1");
