Fix ReplicationSystemLowLevelTests compile errors.

#rnx

diff --git a/Engine/Source/Programs/ReplicationSystemLowLevelTests/ReplicationSystemLowLevelTests.Build.cs b/Engine/Source/Programs/ReplicationSystemLowLevelTests/ReplicationSystemLowLevelTests.Build.cs
--- a/Engine/Source/Programs/ReplicationSystemLowLevelTests/ReplicationSystemLowLevelTests.Build.cs
+++ b/Engine/Source/Programs/ReplicationSystemLowLevelTests/ReplicationSystemLowLevelTests.Build.cs
@@ -32,11 +32,7 @@ public class ReplicationSystemLowLevelTests : TestModuleRules
 		PrivateIncludePathModuleNames.AddRange(
 			new string[]
 			{
-				"Engine",
-				"RHI",
 				"SlateCore",
-				"InputCore",
-				"ApplicationCore"
 			}
 		);
 	}
diff --git a/Engine/Source/Programs/ReplicationSystemLowLevelTests/ReplicationSystemLowLevelTests.Target.cs b/Engine/Source/Programs/ReplicationSystemLowLevelTests/ReplicationSystemLowLevelTests.Target.cs
--- a/Engine/Source/Programs/ReplicationSystemLowLevelTests/ReplicationSystemLowLevelTests.Target.cs
+++ b/Engine/Source/Programs/ReplicationSystemLowLevelTests/ReplicationSystemLowLevelTests.Target.cs
@@ -7,13 +7,14 @@ public class ReplicationSystemLowLevelTestsTarget : TestTargetRules
 {
 	public ReplicationSystemLowLevelTestsTarget(TargetInfo Target) : base(Target)
 	{
-		bCompileAgainstEngine = true;
-		bCompileAgainstApplicationCore = true;
-		bUsesSlate = false;
-
+		bTestsRequireEngine = true;
+		bTestsRequireApplicationCore = true;
+		bTestsRequireCoreUObject = true;
 		bUsePlatformFileStub = true;
 		bMockEngineDefaults = true;
 
+		bUsesSlate = false;
+
 		IncludeOrderVersion = EngineIncludeOrderVersion.Latest;
 
 		// Network config
