Fix ReplicationSystemTest compile errors.

#rnx
#Synced-CL 48899922

diff --git a/Engine/Source/Programs/ReplicationSystemLowLevelTests/Private/ReplicationSystemLowLevelTests.cpp b/Engine/Source/Programs/ReplicationSystemLowLevelTests/Private/ReplicationSystemLowLevelTests.cpp
--- a/Engine/Source/Programs/ReplicationSystemLowLevelTests/Private/ReplicationSystemLowLevelTests.cpp
+++ b/Engine/Source/Programs/ReplicationSystemLowLevelTests/Private/ReplicationSystemLowLevelTests.cpp
@@ -74,7 +74,7 @@ GROUP_BEFORE_GLOBAL(Catch::DefaultGroup)
 		LOG_SCOPE_VERBOSITY_OVERRIDE(LogUObjectGlobals, ELogVerbosity::Fatal);
 		LOG_SCOPE_VERBOSITY_OVERRIDE(LogStreaming, ELogVerbosity::Error);
 		InitAll(true, true);
-		AutoRTFM::InitializeForUE();
+		AutoRTFM::InitializeForUE(nullptr);
 	}
 
 	for (const FName ModuleToLoad : {"IrisCore", "ReplicationSystemTestPlugin"})
diff --git a/Engine/Source/Programs/ReplicationSystemTest/Private/ReplicationSystemTest.cpp b/Engine/Source/Programs/ReplicationSystemTest/Private/ReplicationSystemTest.cpp
--- a/Engine/Source/Programs/ReplicationSystemTest/Private/ReplicationSystemTest.cpp
+++ b/Engine/Source/Programs/ReplicationSystemTest/Private/ReplicationSystemTest.cpp
@@ -130,7 +130,7 @@ static void PreInit()
 #endif
 	FPlatformMemory::Init();
 	
-	AutoRTFM::InitializeForUE();
+	AutoRTFM::InitializeForUE(nullptr);
 
 	// Initialize trace
 	FString Parameter;
